fixup DIASDK paths for Windows Arm64

#rnx

diff --git a/Engine/Source/Developer/Windows/LiveCodingServer/LiveCodingServer.Build.cs b/Engine/Source/Developer/Windows/LiveCodingServer/LiveCodingServer.Build.cs
--- a/Engine/Source/Developer/Windows/LiveCodingServer/LiveCodingServer.Build.cs
+++ b/Engine/Source/Developer/Windows/LiveCodingServer/LiveCodingServer.Build.cs
@@ -27,10 +27,11 @@ public class LiveCodingServer : ModuleRules
 				throw new System.Exception("Unable to find DIA SDK directory");
 			}
 
+			string Arch = (Target.Architecture == UnrealArch.Arm64) ? "arm64" : "amd64";
 			ExtraRootPath = ("DiaSdkDir", DiaSdkDir);
 			PrivateIncludePaths.Add(Path.Combine(DiaSdkDir, "include"));
-			PublicAdditionalLibraries.Add(Path.Combine(DiaSdkDir, "lib", "amd64", "diaguids.lib"));
-			RuntimeDependencies.Add("$(TargetOutputDir)/msdia140.dll", Path.Combine(DiaSdkDir, "bin", "amd64", "msdia140.dll"));
+			PublicAdditionalLibraries.Add(Path.Combine(DiaSdkDir, "lib", Arch, "diaguids.lib"));
+			RuntimeDependencies.Add("$(TargetOutputDir)/msdia140.dll", Path.Combine(DiaSdkDir, "bin", Arch, "msdia140.dll"));
 		}
 
 		if (Target.Configuration == UnrealTargetConfiguration.Debug)
diff --git a/Engine/Source/Programs/DumpSyms/DumpSyms.Build.cs b/Engine/Source/Programs/DumpSyms/DumpSyms.Build.cs
--- a/Engine/Source/Programs/DumpSyms/DumpSyms.Build.cs
+++ b/Engine/Source/Programs/DumpSyms/DumpSyms.Build.cs
@@ -16,13 +16,15 @@ public class DumpSyms : ModuleRules
 		if (Target.Platform == UnrealTargetPlatform.Win64)
 		{
 			string ToolChainDir = Target.WindowsPlatform.ToolChainDir;
+			string ToolChainArch = (Target.Architecture == UnrealArch.Arm64) ? "arm64" : "x64";
 			PrivateIncludePaths.Add(Path.Combine(ToolChainDir, "atlmfc", "include"));
-			PublicAdditionalLibraries.Add(Path.Combine(ToolChainDir, "atlmfc", "lib", "x64", "atls.lib"));
+			PublicAdditionalLibraries.Add(Path.Combine(ToolChainDir, "atlmfc", "lib", ToolChainArch, "atls.lib"));
 
 			string DiaSdkDir = Target.WindowsPlatform.DiaSdkDir;
+			string DiaArch = (Target.Architecture == UnrealArch.Arm64) ? "arm64" : "amd64";
 			PrivateIncludePaths.Add(Path.Combine(DiaSdkDir, "include"));
-			PublicAdditionalLibraries.Add(Path.Combine(DiaSdkDir, "lib", "amd64", "diaguids.lib"));
-			RuntimeDependencies.Add("$(TargetOutputDir)/msdia140.dll", Path.Combine(DiaSdkDir, "bin", "amd64", "msdia140.dll"));
+			PublicAdditionalLibraries.Add(Path.Combine(DiaSdkDir, "lib", DiaArch, "diaguids.lib"));
+			RuntimeDependencies.Add("$(TargetOutputDir)/msdia140.dll", Path.Combine(DiaSdkDir, "bin", DiaArch, "msdia140.dll"));
 
 			string WindowsSdkDir = Target.WindowsPlatform.WindowsSdkDir;
 			PublicAdditionalLibraries.Add(Path.Combine(WindowsSdkDir, "Lib", Target.WindowsPlatform.WindowsSdkVersion, "um", "x64", "imagehlp.lib"));
