arm64 version of msdia140.dll is copied to the arm64 SxS subfolder

#rnx

diff --git a/Engine/Source/Developer/Windows/LiveCodingServer/LiveCodingServer.Build.cs b/Engine/Source/Developer/Windows/LiveCodingServer/LiveCodingServer.Build.cs
--- a/Engine/Source/Developer/Windows/LiveCodingServer/LiveCodingServer.Build.cs
+++ b/Engine/Source/Developer/Windows/LiveCodingServer/LiveCodingServer.Build.cs
@@ -28,10 +28,11 @@ public class LiveCodingServer : ModuleRules
 			}
 
 			string Arch = (Target.Architecture == UnrealArch.Arm64) ? "arm64" : "amd64";
+			string RuntimeArchDir = (Target.Architecture == UnrealArch.Arm64) ? "arm64" : "";
 			ExtraRootPath = ("DiaSdkDir", DiaSdkDir);
 			PrivateIncludePaths.Add(Path.Combine(DiaSdkDir, "include"));
 			PublicAdditionalLibraries.Add(Path.Combine(DiaSdkDir, "lib", Arch, "diaguids.lib"));
-			RuntimeDependencies.Add("$(TargetOutputDir)/msdia140.dll", Path.Combine(DiaSdkDir, "bin", Arch, "msdia140.dll"));
+			RuntimeDependencies.Add(Path.Combine("$(TargetOutputDir)", RuntimeArchDir, "msdia140.dll"), Path.Combine(DiaSdkDir, "bin", Arch, "msdia140.dll"));
 		}
 
 		if (Target.Configuration == UnrealTargetConfiguration.Debug)
diff --git a/Engine/Source/Programs/DumpSyms/DumpSyms.Build.cs b/Engine/Source/Programs/DumpSyms/DumpSyms.Build.cs
--- a/Engine/Source/Programs/DumpSyms/DumpSyms.Build.cs
+++ b/Engine/Source/Programs/DumpSyms/DumpSyms.Build.cs
@@ -22,9 +22,10 @@ public class DumpSyms : ModuleRules
 
 			string DiaSdkDir = Target.WindowsPlatform.DiaSdkDir;
 			string DiaArch = (Target.Architecture == UnrealArch.Arm64) ? "arm64" : "amd64";
+			string RuntimeArchDir = (Target.Architecture == UnrealArch.Arm64) ? "arm64" : "";
 			PrivateIncludePaths.Add(Path.Combine(DiaSdkDir, "include"));
 			PublicAdditionalLibraries.Add(Path.Combine(DiaSdkDir, "lib", DiaArch, "diaguids.lib"));
-			RuntimeDependencies.Add("$(TargetOutputDir)/msdia140.dll", Path.Combine(DiaSdkDir, "bin", DiaArch, "msdia140.dll"));
+			RuntimeDependencies.Add(Path.Combine("$(TargetOutputDir)", RuntimeArchDir, "msdia140.dll"), Path.Combine(DiaSdkDir, "bin", DiaArch, "msdia140.dll"));
 
 			string WindowsSdkDir = Target.WindowsPlatform.WindowsSdkDir;
 			PublicAdditionalLibraries.Add(Path.Combine(WindowsSdkDir, "Lib", Target.WindowsPlatform.WindowsSdkVersion, "um", "x64", "imagehlp.lib"));
